RigLogicAnimNext to RigLogicUAF & Added node template

* Renamed RigLogicAnimNext to RigLogicUAF
* Including all internal classes, structs etc.
* Created new UncookedOnly module
* Created a graph node template for the RigLogic node
* Added icon for the node template
* Updated the MH anim graphs to the new node template

#virtualized

diff --git a/Engine/Plugins/Experimental/RigLogicAnimNext/RigLogicAnimNext.uplugin b/Engine/Plugins/Experimental/RigLogicAnimNext/RigLogicAnimNext.uplugin
deleted file mode 100644
diff --git a/Engine/Plugins/Experimental/RigLogicAnimNext/Source/RigLogicAnimNext/Private/RigLogicAnimNext.cpp b/Engine/Plugins/Experimental/RigLogicAnimNext/Source/RigLogicAnimNext/Private/RigLogicAnimNext.cpp
deleted file mode 100644
diff --git a/Engine/Plugins/Experimental/RigLogicAnimNext/Source/RigLogicAnimNext/Private/RigLogicInstanceData.cpp b/Engine/Plugins/Experimental/RigLogicAnimNext/Source/RigLogicAnimNext/Private/RigLogicInstanceData.cpp
deleted file mode 100644
diff --git a/Engine/Plugins/Experimental/RigLogicAnimNext/Source/RigLogicAnimNext/Private/RigLogicInstanceDataPool.cpp b/Engine/Plugins/Experimental/RigLogicAnimNext/Source/RigLogicAnimNext/Private/RigLogicInstanceDataPool.cpp
deleted file mode 100644
diff --git a/Engine/Plugins/Experimental/RigLogicAnimNext/Source/RigLogicAnimNext/Private/RigLogicTask.cpp b/Engine/Plugins/Experimental/RigLogicAnimNext/Source/RigLogicAnimNext/Private/RigLogicTask.cpp
deleted file mode 100644
diff --git a/Engine/Plugins/Experimental/RigLogicAnimNext/Source/RigLogicAnimNext/Private/RigLogicTrait.cpp b/Engine/Plugins/Experimental/RigLogicAnimNext/Source/RigLogicAnimNext/Private/RigLogicTrait.cpp
deleted file mode 100644
diff --git a/Engine/Plugins/Experimental/RigLogicAnimNext/Source/RigLogicAnimNext/Public/RigLogicAnimNext.h b/Engine/Plugins/Experimental/RigLogicAnimNext/Source/RigLogicAnimNext/Public/RigLogicAnimNext.h
deleted file mode 100644
diff --git a/Engine/Plugins/Experimental/RigLogicAnimNext/Source/RigLogicAnimNext/Public/RigLogicInstanceData.h b/Engine/Plugins/Experimental/RigLogicAnimNext/Source/RigLogicAnimNext/Public/RigLogicInstanceData.h
deleted file mode 100644
diff --git a/Engine/Plugins/Experimental/RigLogicAnimNext/Source/RigLogicAnimNext/Public/RigLogicInstanceDataPool.h b/Engine/Plugins/Experimental/RigLogicAnimNext/Source/RigLogicAnimNext/Public/RigLogicInstanceDataPool.h
deleted file mode 100644
diff --git a/Engine/Plugins/Experimental/RigLogicAnimNext/Source/RigLogicAnimNext/Public/RigLogicTask.h b/Engine/Plugins/Experimental/RigLogicAnimNext/Source/RigLogicAnimNext/Public/RigLogicTask.h
deleted file mode 100644
diff --git a/Engine/Plugins/Experimental/RigLogicAnimNext/Source/RigLogicAnimNext/Public/RigLogicTrait.h b/Engine/Plugins/Experimental/RigLogicAnimNext/Source/RigLogicAnimNext/Public/RigLogicTrait.h
deleted file mode 100644
diff --git a/Engine/Plugins/Experimental/RigLogicAnimNext/Source/RigLogicAnimNext/RigLogicAnimNext.Build.cs b/Engine/Plugins/Experimental/RigLogicAnimNext/Source/RigLogicAnimNext/RigLogicAnimNext.Build.cs
deleted file mode 100644
diff --git a/Engine/Plugins/Experimental/RigLogicUAF/RigLogicUAF.uplugin b/Engine/Plugins/Experimental/RigLogicUAF/RigLogicUAF.uplugin
new file mode 100644
--- /dev/null
+++ b/Engine/Plugins/Experimental/RigLogicUAF/RigLogicUAF.uplugin
@@ -0,0 +1,45 @@
+{
+	"FileVersion": 3,
+	"Version": 1,
+	"VersionName": "1.0",
+	"FriendlyName": "RigLogic for UAF",
+	"Description": "RigLogic for UAF",
+	"Category": "Animation",
+	"CreatedBy": "Epic Games Inc.",
+	"CreatedByURL": "",
+	"DocsURL": "",
+	"MarketplaceURL": "",
+	"SupportURL": "",
+	"CanContainContent": true,
+	"IsBetaVersion": false,
+	"IsExperimentalVersion": true,
+	"Installed": false,
+	"Modules":
+	[
+		{
+			"Name": "RigLogicUAF",
+			"Type": "Runtime",
+			"LoadingPhase": "Default"
+		},
+		{
+			"Name": "RigLogicUAFUncookedOnly",
+			"Type": "UncookedOnly",
+			"LoadingPhase": "Default"
+		}
+	],
+	"Plugins":
+	[
+		{
+			"Name": "RigLogic",
+			"Enabled": true
+		},
+		{
+			"Name": "UAF",
+			"Enabled": true
+		},
+		{
+			"Name": "UAFAnimGraph",
+			"Enabled": true
+		}
+	]
+}
\ No newline at end of file
diff --git a/Engine/Plugins/Experimental/RigLogicUAF/Source/RigLogicUAF/Private/RigLogicInstanceData.cpp b/Engine/Plugins/Experimental/RigLogicUAF/Source/RigLogicUAF/Private/RigLogicInstanceData.cpp
new file mode 100644
--- /dev/null
+++ b/Engine/Plugins/Experimental/RigLogicUAF/Source/RigLogicUAF/Private/RigLogicInstanceData.cpp
@@ -0,0 +1,138 @@
+// Copyright Epic Games, Inc. All Rights Reserved.
+
+#include "RigLogicInstanceData.h"
+#include "RigLogicUAF.h"
+
+#include "Animation/Skeleton.h"
+#include "Engine/SkeletalMesh.h"
+#include "DNAAsset.h"
+#include "LODPose.h"
+#include "SharedRigRuntimeContext.h"
+
+namespace UE::UAF
+{
+	void FRigLogicInstanceData::Init(const UE::UAF::FReferencePose* ReferencePose)
+	{
+		LLM_SCOPE_BYNAME(TEXT("Animation/RigLogic"));
+
+		if (!ReferencePose)
+		{
+			UE_LOG(LogRigLogicUAF, Error, TEXT("Reference pose invalid."));
+			return;
+		}
+
+		USkeletalMesh* SkeletalMesh = const_cast<USkeletalMesh*>(ReferencePose->SkeletalMesh.Get());
+		if (!SkeletalMesh)
+		{
+			UE_LOG(LogRigLogicUAF, Error, TEXT("No skeletal mesh assigned to reference pose."));
+			return;
+		}
+
+		const USkeleton* Skeleton = SkeletalMesh->GetSkeleton();
+		if (!Skeleton)
+		{
+			UE_LOG(LogRigLogicUAF, Error, TEXT("No skeleton assigned to the skeletal mesh."));
+			return;
+		}
+
+		UDNAAsset* DNAAsset = Cast<UDNAAsset>(SkeletalMesh->GetAssetUserDataOfClass(UDNAAsset::StaticClass()));
+		if (!DNAAsset)
+		{
+			UE_LOG(LogRigLogicUAF, Warning, TEXT("No DNA asset assigned to the skeletal mesh."));
+			return;
+		}
+
+		TSharedPtr<FSharedRigRuntimeContext> SharedRigRuntimeContext = DNAAsset->GetRigRuntimeContext();
+		if (!SharedRigRuntimeContext.IsValid())
+		{
+			UE_LOG(LogRigLogicUAF, Error, TEXT("Can't get the shared rig runtime context."));
+			return;
+		}
+
+		if (CachedRigRuntimeContext != SharedRigRuntimeContext)
+		{
+			CachedRigRuntimeContext = SharedRigRuntimeContext;
+		}
+
+		CachedDNAIndexMapping = DNAAsset->GetDNAIndexMapping(Skeleton, SkeletalMesh);
+
+		NumLODs = CachedRigRuntimeContext->VariableJointIndicesPerLOD.Num();
+
+		InitBoneIndexMapping(ReferencePose);
+		InitSparseAndDenseDriverJointMapping(ReferencePose);
+	}
+
+	void FRigLogicInstanceData::InitBoneIndexMapping(const UE::UAF::FReferencePose* ReferencePose)
+	{
+		LLM_SCOPE_BYNAME(TEXT("Animation/RigLogic"));
+
+		RigLogicToSkeletonBoneIndexMappingPerLOD.Empty();
+		RigLogicToSkeletonBoneIndexMappingPerLOD.SetNum(NumLODs);
+
+		const TArrayView<const FBoneIndexType> MeshToPoseBoneIndexMap = ReferencePose->GetMeshBoneIndexToLODBoneIndexMap(); // mesh -> pose
+
+		for (uint32 LODLevel = 0; LODLevel < NumLODs; ++LODLevel)
+		{
+			const TArray<uint16>& VariableJointIndices = CachedRigRuntimeContext->VariableJointIndicesPerLOD[LODLevel].Values;
+			RigLogicToSkeletonBoneIndexMappingPerLOD[LODLevel].Reserve(VariableJointIndices.Num());
+
+			for (const uint16 RigLogicJointIndex : VariableJointIndices)
+			{
+				// Get the mesh bone index from the RigLogic joint index. Index originating from a FindBoneName() on the mesh's reference skeleton.
+				const FMeshPoseBoneIndex MeshBoneIndex = CachedDNAIndexMapping->JointsMapDNAIndicesToMeshPoseBoneIndices[RigLogicJointIndex];
+				if (MeshBoneIndex.IsValid())
+				{
+					// Convert the skeleton bone index to a pose bone index for the given LOD level.
+					const int32 PoseBoneIndex = MeshToPoseBoneIndexMap[MeshBoneIndex.GetInt()];
+					if (ReferencePose->IsBoneEnabled(PoseBoneIndex, LODLevel))
+					{
+						RigLogicToSkeletonBoneIndexMappingPerLOD[LODLevel].Add({ RigLogicJointIndex, PoseBoneIndex });
+					}
+				}
+				else
+				{
+					UE_LOG(LogTemp, Warning, TEXT("Could not find bone in skeleton for RigLogic joint with index %i."), RigLogicJointIndex);
+				}
+			}
+		}
+	}
+
+	void FRigLogicInstanceData::InitSparseAndDenseDriverJointMapping(const UE::UAF::FReferencePose* ReferencePose)
+	{
+		LLM_SCOPE_BYNAME(TEXT("Animation/RigLogic"));
+
+		// Populate driver joint to raw control attribute mapping (used to feed RigLogic with inputs from the joint hierarchy)
+		SparseDriverJointsToControlAttributesMapPerLOD.Empty();
+		DenseDriverJointsToControlAttributesMapPerLOD.Empty();
+
+		SparseDriverJointsToControlAttributesMapPerLOD.SetNum(NumLODs);
+		DenseDriverJointsToControlAttributesMapPerLOD.SetNum(NumLODs);
+
+		const TArrayView<const FBoneIndexType> MeshToPoseBoneIndexMap = ReferencePose->GetMeshBoneIndexToLODBoneIndexMap(); // mesh -> pose
+		
+		for (uint32 LODLevel = 0; LODLevel < NumLODs; ++LODLevel)
+		{
+			// Sparse mapping will likely remain empty so no reservation happens
+			for (const auto& Mapping : CachedDNAIndexMapping->DriverJointsToControlAttributesMap)
+			{
+				const FMeshPoseBoneIndex MeshBoneIndex = Mapping.MeshPoseBoneIndex;
+				if (MeshBoneIndex.IsValid())
+				{
+					// Convert the mesh bone index to a pose bone index for the given LOD level.
+					const int32 PoseBoneIndex = MeshToPoseBoneIndexMap[MeshBoneIndex.GetInt()];
+					if (ReferencePose->IsBoneEnabled(PoseBoneIndex, LODLevel))
+					{
+						if ((Mapping.RotationX != INDEX_NONE) && (Mapping.RotationY != INDEX_NONE) && (Mapping.RotationZ != INDEX_NONE) && (Mapping.RotationW != INDEX_NONE))
+						{
+							DenseDriverJointsToControlAttributesMapPerLOD[LODLevel].Add({ PoseBoneIndex, Mapping.DNAJointIndex, Mapping.RotationX, Mapping.RotationY, Mapping.RotationZ, Mapping.RotationW });
+						}
+						else
+						{
+							SparseDriverJointsToControlAttributesMapPerLOD[LODLevel].Add({ PoseBoneIndex, Mapping.DNAJointIndex, Mapping.RotationX, Mapping.RotationY, Mapping.RotationZ, Mapping.RotationW });
+						}
+					}
+				}
+			}
+		}
+	}
+} // namespace UE::UAF
diff --git a/Engine/Plugins/Experimental/RigLogicUAF/Source/RigLogicUAF/Private/RigLogicInstanceDataPool.cpp b/Engine/Plugins/Experimental/RigLogicUAF/Source/RigLogicUAF/Private/RigLogicInstanceDataPool.cpp
new file mode 100644
--- /dev/null
+++ b/Engine/Plugins/Experimental/RigLogicUAF/Source/RigLogicUAF/Private/RigLogicInstanceDataPool.cpp
@@ -0,0 +1,96 @@
+// Copyright Epic Games, Inc. All Rights Reserved.
+
+#include "RigLogicInstanceDataPool.h"
+#include "RigLogicUAF.h"
+
+namespace UE::UAF
+{
+	TSharedPtr<FRigLogicInstanceData> FRigLogicInstanceDataPool::RequestData(const UE::UAF::FReferencePose* InReferencePose)
+	{
+		LLM_SCOPE_BYNAME(TEXT("Animation/RigLogic"));
+
+		// Only lock while searching
+		{
+			FScopeLock Lock(&PoolAccessCriticalSection);
+
+			if (TArray<TSharedPtr<FRigLogicInstanceData>>* Found = Datas.Find(InReferencePose->SkeletalMesh))
+			{
+				if (!Found->IsEmpty())
+				{
+					// We have pre-allocated data already, use one that is already there.
+					return Found->Pop();
+				}
+			}
+		}
+		// We don't need to keep things locked here as we're creating a new object
+		// Initializing a new data is not a lightweight operation, so be sure to close the lock first.
+
+		// We don't have any pool datas in the array yet. Create a new one.
+		TSharedPtr<FRigLogicInstanceData> NewObject = TSharedPtr<FRigLogicInstanceData>(new FRigLogicInstanceData());
+		NewObject->Init(InReferencePose);
+		return NewObject;
+	}
+
+	void FRigLogicInstanceDataPool::FreeData(TWeakObjectPtr<const USkeletalMesh> SkeletalMesh, TSharedPtr<FRigLogicInstanceData> InData)
+	{
+		LLM_SCOPE_BYNAME(TEXT("Animation/RigLogic"));
+		FScopeLock Lock(&PoolAccessCriticalSection);
+
+		// Do we have an element for the given skel mesh?
+		TArray<TSharedPtr<FRigLogicInstanceData>>* Found = Datas.Find(SkeletalMesh);
+		if (!Found)
+		{
+			// We don't. Create a new one and add the instance data to it.
+			TArray<TSharedPtr<FRigLogicInstanceData>> NewArray;
+			NewArray.Add(InData);
+			Datas.Add(SkeletalMesh, NewArray);
+		}
+		else
+		{
+			// Found one, just add the instance data back in.
+			Found->Add(InData);
+		}
+	}
+
+	void FRigLogicInstanceDataPool::GarbageCollect()
+	{
+		LLM_SCOPE_BYNAME(TEXT("Animation/RigLogic"));
+		FScopeLock Lock(&PoolAccessCriticalSection);
+
+		for (auto Iterator = Datas.CreateIterator(); Iterator; ++Iterator)
+		{
+			// Remove cached instance datas for skeletal meshes that are not loaded anymore.
+			if (!Iterator->Key.IsValid())
+			{
+				Iterator.RemoveCurrent();
+			}
+		}
+	}
+
+	void FRigLogicInstanceDataPool::Log()
+	{
+		FScopeLock Lock(&PoolAccessCriticalSection);
+
+		UE_LOG(LogRigLogicUAF, Display, TEXT("Pool data:"));
+		uint32 Counter = 0;
+		for (auto Iterator = Datas.CreateIterator(); Iterator; ++Iterator)
+		{
+			if (Iterator->Key.IsValid())
+			{
+				const USkeletalMesh* SkeletalMesh = Iterator->Key.Get();
+				UE_LOG(LogRigLogicUAF, Display, TEXT("   - Skeletal Mesh %s:"), *SkeletalMesh->GetPathName());
+
+				const TArray<TSharedPtr<FRigLogicInstanceData>>& InstanceDataArray = Iterator->Value;
+				for (int32 i = 0; i < InstanceDataArray.Num(); ++i)
+				{
+					UE_LOG(LogRigLogicUAF, Display, TEXT("      + InstanceData %i (0x%d):"), i, InstanceDataArray[i].Get());
+				}
+			}
+			else
+			{
+				UE_LOG(LogRigLogicUAF, Warning, TEXT("Entry %i linked to an invalid skeletal mesh."), Counter);
+			}
+			Counter++;
+		}
+	}
+} // namespace UE::UAF
\ No newline at end of file
diff --git a/Engine/Plugins/Experimental/RigLogicUAF/Source/RigLogicUAF/Private/RigLogicTask.cpp b/Engine/Plugins/Experimental/RigLogicUAF/Source/RigLogicUAF/Private/RigLogicTask.cpp
new file mode 100644
--- /dev/null
+++ b/Engine/Plugins/Experimental/RigLogicUAF/Source/RigLogicUAF/Private/RigLogicTask.cpp
@@ -0,0 +1,245 @@
+// Copyright Epic Games, Inc. All Rights Reserved.
+
+#include "RigLogicTask.h"
+
+#include "EvaluationVM/EvaluationVM.h"
+#include "EvaluationVM/KeyframeState.h"
+#include "TransformArrayOperations.h"
+
+#include "RigLogicUAF.h"
+#include "DNAIndexMapping.h"
+#include "RigLogic.h"
+#include "RigInstance.h"
+#include "SharedRigRuntimeContext.h"
+#include "RigLogicInstanceDataPool.h"
+
+#include <tdm/TDM.h>
+
+DEFINE_STAT(STAT_UAF_Task_RigLogic);
+
+FUAFRigLogicTask FUAFRigLogicTask::Make(UE::UAF::FRigLogicTrait::FInstanceData* InstanceData)
+{
+	FUAFRigLogicTask Result;
+	Result.TraitInstanceData = InstanceData;
+	return Result;
+}
+
+void FUAFRigLogicTask::Execute(UE::UAF::FEvaluationVM& VM) const
+{
+	SCOPE_CYCLE_COUNTER(STAT_UAF_Task_RigLogic);
+
+	// Pop pose, we'll re-use it for our result
+	TUniquePtr<UE::UAF::FKeyframeState> Keyframe;
+	if (!VM.PopValue(UE::UAF::KEYFRAME_STACK_NAME, Keyframe))
+	{
+		// We have no inputs, nothing to do.
+		return;
+	}
+
+	// Lazy-init instance data.
+	UE::UAF::FRigLogicModule& Module = FModuleManager::GetModuleChecked<UE::UAF::FRigLogicModule>("RigLogicUAF");
+	TSharedPtr<UE::UAF::FRigLogicInstanceData> PoolInstanceData = Module.DataPool.RequestData(Keyframe->Pose.RefPose);
+
+	if (const FSharedRigRuntimeContext* RigRuntimeContext = PoolInstanceData->CachedRigRuntimeContext.Get())
+	{
+		if (!TraitInstanceData->RigInstance)
+		{
+			TraitInstanceData->RigInstance = TUniquePtr<FRigInstance>(new FRigInstance(RigRuntimeContext->RigLogic.Get()));
+		}
+		FRigInstance* RigInstance = TraitInstanceData->RigInstance.Get();
+
+		const FDNAIndexMapping* DNAIndexMapping = PoolInstanceData->CachedDNAIndexMapping.Get();
+		if (RigInstance && RigRuntimeContext && DNAIndexMapping)
+		{
+			if (const FRigLogic* RigLogic = RigRuntimeContext->RigLogic.Get())
+			{
+				const int32 LODLevel = FMath::Clamp(Keyframe->Pose.LODLevel, 0, (int32)PoolInstanceData->NumLODs - 1);
+
+				// 1. Feed input curves to RigLogic.
+				FBlendedCurve& Curves = Keyframe->Curves;
+				const TArrayView<const float> NeutralJointValues = PoolInstanceData->CachedRigRuntimeContext->RigLogic->GetNeutralJointValues();
+
+				UpdateControlCurves(RigInstance,
+					Curves,
+					DNAIndexMapping->ControlAttributeCurves,
+					DNAIndexMapping->NeuralNetworkMaskCurves,
+					NeutralJointValues,
+					PoolInstanceData->SparseDriverJointsToControlAttributesMapPerLOD[LODLevel],
+					PoolInstanceData->DenseDriverJointsToControlAttributesMapPerLOD[LODLevel],
+					Keyframe->Pose);
+
+				// 2. Evaluate RigLogic.
+				RigInstance->SetLOD(LODLevel);
+				RigLogic->Calculate(RigInstance);
+
+				// 3. Update UAF bone transforms and curves with RigLogic output.
+				TArrayView<const uint16> VariableJointIndices = RigRuntimeContext->VariableJointIndicesPerLOD[LODLevel].Values;
+				const TArrayView<const UE::UAF::FRigLogicBoneMapping> RigLogicIndexToMeshIndexMapping = PoolInstanceData->RigLogicToSkeletonBoneIndexMappingPerLOD[LODLevel];
+				UpdateJoints(RigLogicIndexToMeshIndexMapping,
+					VariableJointIndices,
+					NeutralJointValues,
+					RigInstance->GetJointOutputs(),
+					Keyframe->Pose);
+
+				UpdateBlendShapeCurves(DNAIndexMapping->MorphTargetCurvesPerLOD[LODLevel], RigInstance->GetBlendShapeOutputs(), Curves);
+				UpdateAnimMapCurves(DNAIndexMapping->MaskMultiplierCurvesPerLOD[LODLevel], RigInstance->GetAnimatedMapOutputs(), Curves);
+			}
+		}
+	}
+
+	Module.DataPool.FreeData(Keyframe->Pose.RefPose->SkeletalMesh, PoolInstanceData);
+
+	// Push our result back
+	VM.PushValue(UE::UAF::KEYFRAME_STACK_NAME, MoveTemp(Keyframe));
+}
+
+void FUAFRigLogicTask::UpdateControlCurves(FRigInstance* InRigInstance,
+	const FBlendedCurve& Curves,
+	const FCachedIndexedCurve& ControlAttributeCurves,
+	const FCachedIndexedCurve& NeuralNetworkMaskCurves,
+	const TArrayView<const float> NeutralJointValues,
+	const TArrayView<const UE::UAF::FPoseBoneControlAttributeMapping> SparseDriverJointsToControlAttributesMap,
+	const TArrayView<const UE::UAF::FPoseBoneControlAttributeMapping> DenseDriverJointsToControlAttributesMap,
+	const UE::UAF::FLODPoseStack& InputPose) const
+{
+	// Combine control attribute curve with input curve to get indexed curve to apply to rig
+	// Curve elements that dont have a control mapping will have INDEX_NONE as their index
+	UE::Anim::FNamedValueArrayUtils::Union(Curves, ControlAttributeCurves,
+		[this, InRigInstance](const UE::Anim::FCurveElement& InCurveElement, const UE::Anim::FCurveElementIndexed& InControlAttributeCurveElement, UE::Anim::ENamedValueUnionFlags InFlags)
+		{
+			if (InControlAttributeCurveElement.Index != INDEX_NONE)
+			{
+				InRigInstance->SetRawControl(InControlAttributeCurveElement.Index, InCurveElement.Value);
+			}
+		});
+
+	// Combine control attribute curve with input curve to get indexed curve to apply to rig
+	// Curve elements that dont have a control mapping will have INDEX_NONE as their index
+	UE::Anim::FNamedValueArrayUtils::Union(Curves, ControlAttributeCurves,
+		[this, InRigInstance](const UE::Anim::FCurveElement& InCurveElement, const UE::Anim::FCurveElementIndexed& InControlAttributeCurveElement, UE::Anim::ENamedValueUnionFlags InFlags)
+		{
+			if (InControlAttributeCurveElement.Index != INDEX_NONE)
+			{
+				InRigInstance->SetRawControl(InControlAttributeCurveElement.Index, InCurveElement.Value);
+			}
+		});
+
+	const uint16 LODLevel = InRigInstance->GetLOD();
+
+	// The sparse mapping is NOT guaranteed to supply all quaternion attributes, so checks for each attribute mapping are present
+	for (int32 MappingIndex = 0; MappingIndex < SparseDriverJointsToControlAttributesMap.Num(); ++MappingIndex)
+	{
+		const UE::UAF::FPoseBoneControlAttributeMapping& Mapping = SparseDriverJointsToControlAttributesMap[MappingIndex];
+		const int32 PoseBoneIndex = Mapping.PoseBoneIndex;
+
+		const UE::UAF::FTransformSoAAdapterConst TransformAdapter = InputPose.LocalTransforms[PoseBoneIndex];
+
+		// Translation and Scale is currently not used here, so to avoid the overhead of checking them, they are simply ignored.
+		// Should the need arise to use them as well, this code will need adjustment.
+		const FQuat Rotation = TransformAdapter.GetRotation();
+		const int32 AttrIndex = Mapping.DNAJointIndex * ATTR_COUNT_PER_JOINT;
+
+		const tdm::fquat NeutralRotation{ NeutralJointValues[AttrIndex + 3], NeutralJointValues[AttrIndex + 4], NeutralJointValues[AttrIndex + 5], NeutralJointValues[AttrIndex + 6] };
+		const tdm::fquat AbsPoseRotation{ static_cast<float>(Rotation.X), static_cast<float>(Rotation.Y), static_cast<float>(Rotation.Z), static_cast<float>(Rotation.W) };
+		const tdm::fquat DeltaPoseRotation = tdm::inverse(NeutralRotation) * AbsPoseRotation;
+
+		if (Mapping.RotationX != INDEX_NONE)
+		{
+			InRigInstance->SetRawControl(Mapping.RotationX, DeltaPoseRotation.x);
+		}
+		if (Mapping.RotationY != INDEX_NONE)
+		{
+			InRigInstance->SetRawControl(Mapping.RotationY, DeltaPoseRotation.y);
+		}
+		if (Mapping.RotationZ != INDEX_NONE)
+		{
+			InRigInstance->SetRawControl(Mapping.RotationZ, DeltaPoseRotation.z);
+		}
+		if (Mapping.RotationW != INDEX_NONE)
+		{
+			InRigInstance->SetRawControl(Mapping.RotationW, DeltaPoseRotation.w);
+		}
+	}
+
+	// The dense mapping is guaranteed to supply all quaternion attributes, so NO checks for each attribute mapping are present
+	for (int32 MappingIndex = 0; MappingIndex < DenseDriverJointsToControlAttributesMap.Num(); ++MappingIndex)
+	{
+		const UE::UAF::FPoseBoneControlAttributeMapping& Mapping = DenseDriverJointsToControlAttributesMap[MappingIndex];
+		const int32 PoseBoneIndex = Mapping.PoseBoneIndex;
+
+		const UE::UAF::FTransformSoAAdapterConst TransformAdapter = InputPose.LocalTransforms[PoseBoneIndex];
+
+		// Translation and Scale is currently not used here, so to avoid the overhead of checking them, they are simply ignored.
+		// Should the need arise to use them as well, this code will need adjustment.
+		const FQuat Rotation = TransformAdapter.GetRotation();
+		const int32 AttrIndex = Mapping.DNAJointIndex * ATTR_COUNT_PER_JOINT;
+
+		const tdm::fquat NeutralRotation{ NeutralJointValues[AttrIndex + 3], NeutralJointValues[AttrIndex + 4], NeutralJointValues[AttrIndex + 5], NeutralJointValues[AttrIndex + 6] };
+		const tdm::fquat AbsPoseRotation{ static_cast<float>(Rotation.X), static_cast<float>(Rotation.Y), static_cast<float>(Rotation.Z), static_cast<float>(Rotation.W) };
+		const tdm::fquat DeltaPoseRotation = tdm::inverse(NeutralRotation) * AbsPoseRotation;
+
+		InRigInstance->SetRawControl(Mapping.RotationX, DeltaPoseRotation.x);
+		InRigInstance->SetRawControl(Mapping.RotationY, DeltaPoseRotation.y);
+		InRigInstance->SetRawControl(Mapping.RotationZ, DeltaPoseRotation.z);
+		InRigInstance->SetRawControl(Mapping.RotationW, DeltaPoseRotation.w);
+	}
+
+	if (InRigInstance->GetNeuralNetworkCount() != 0)
+	{
+		UE::Anim::FNamedValueArrayUtils::Union(Curves, NeuralNetworkMaskCurves,
+			[this, InRigInstance](const UE::Anim::FCurveElement& InCurveElement, const UE::Anim::FCurveElementIndexed& InControlAttributeCurveElement, UE::Anim::ENamedValueUnionFlags InFlags)
+			{
+				if (InControlAttributeCurveElement.Index != INDEX_NONE)
+				{
+					InRigInstance->SetNeuralNetworkMask(InControlAttributeCurveElement.Index, InCurveElement.Value);
+				}
+			});
+	}
+}
+
+void FUAFRigLogicTask::UpdateJoints(const TArrayView<const UE::UAF::FRigLogicBoneMapping> RigLogicIndexToMeshIndexMapping,
+	const TArrayView<const uint16> VariableJointIndices,
+	const TArrayView<const float> NeutralJointValues,
+	const TArrayView<const float> DeltaJointValues,
+	UE::UAF::FLODPoseStack& OutputPose)
+{
+	const float* N = NeutralJointValues.GetData();
+	const float* D = DeltaJointValues.GetData();
+
+	for (const UE::UAF::FRigLogicBoneMapping& Mapping : RigLogicIndexToMeshIndexMapping)
+	{
+		const uint16 RigLogicJointIndex = Mapping.RigLogicJointIndex;
+		const int32 PoseBoneIndex = Mapping.PoseBoneIndex;
+		const uint16 AttrIndex = RigLogicJointIndex * ATTR_COUNT_PER_JOINT;
+
+		UE::UAF::FTransformSoAAdapter TransformAdapter = OutputPose.LocalTransforms[PoseBoneIndex];
+
+		TransformAdapter.SetTranslation(FVector((N[AttrIndex + 0] + D[AttrIndex + 0]), (N[AttrIndex + 1] + D[AttrIndex + 1]), (N[AttrIndex + 2] + D[AttrIndex + 2])));
+		TransformAdapter.SetRotation(FQuat(N[AttrIndex + 3], N[AttrIndex + 4], N[AttrIndex + 5], N[AttrIndex + 6]) * FQuat(D[AttrIndex + 3], D[AttrIndex + 4], D[AttrIndex + 5], D[AttrIndex + 6]));
+		TransformAdapter.SetScale3D(FVector((N[AttrIndex + 7] + D[AttrIndex + 7]), (N[AttrIndex + 8] + D[AttrIndex + 8]), (N[AttrIndex + 9] + D[AttrIndex + 9])));
+	}
+}
+
+void FUAFRigLogicTask::UpdateBlendShapeCurves(const FCachedIndexedCurve& MorphTargetCurves, const TArrayView<const float> BlendShapeValues, FBlendedCurve& OutputCurves)
+{
+	UE::Anim::FNamedValueArrayUtils::Union(OutputCurves, MorphTargetCurves, [&BlendShapeValues](UE::Anim::FCurveElement& InOutResult, const UE::Anim::FCurveElementIndexed& InSource, UE::Anim::ENamedValueUnionFlags InFlags)
+		{
+			if (BlendShapeValues.IsValidIndex(InSource.Index))
+			{
+				InOutResult.Value = BlendShapeValues[InSource.Index];
+				InOutResult.Flags |= UE::Anim::ECurveElementFlags::MorphTarget;
+			}
+		});
+}
+
+void FUAFRigLogicTask::UpdateAnimMapCurves(const FCachedIndexedCurve& MaskMultiplierCurves, const TArrayView<const float> AnimMapOutputs, FBlendedCurve& OutputCurves)
+{
+	UE::Anim::FNamedValueArrayUtils::Union(OutputCurves, MaskMultiplierCurves, [&AnimMapOutputs](UE::Anim::FCurveElement& InOutResult, const UE::Anim::FCurveElementIndexed& InSource, UE::Anim::ENamedValueUnionFlags InFlags)
+		{
+			if (AnimMapOutputs.IsValidIndex(InSource.Index))
+			{
+				InOutResult.Value = AnimMapOutputs[InSource.Index];
+				InOutResult.Flags |= UE::Anim::ECurveElementFlags::Material;
+			}
+		});
+}
\ No newline at end of file
diff --git a/Engine/Plugins/Experimental/RigLogicUAF/Source/RigLogicUAF/Private/RigLogicTrait.cpp b/Engine/Plugins/Experimental/RigLogicUAF/Source/RigLogicUAF/Private/RigLogicTrait.cpp
new file mode 100644
--- /dev/null
+++ b/Engine/Plugins/Experimental/RigLogicUAF/Source/RigLogicUAF/Private/RigLogicTrait.cpp
@@ -0,0 +1,73 @@
+// Copyright Epic Games, Inc. All Rights Reserved.
+
+#include "RigLogicTrait.h"
+
+#include "RigLogicUAF.h"
+#include "RigLogicTask.h"
+#include "TraitCore/ExecutionContext.h"
+
+namespace UE::UAF
+{
+	AUTO_REGISTER_ANIM_TRAIT(FRigLogicTrait)
+
+		// Trait implementation boilerplate
+#define TRAIT_INTERFACE_ENUMERATOR(GeneratorMacro) \
+		GeneratorMacro(IEvaluate) \
+		GeneratorMacro(IHierarchy) \
+		GeneratorMacro(IUpdate) \
+		GeneratorMacro(IUpdateTraversal) \
+
+		GENERATE_ANIM_TRAIT_IMPLEMENTATION(FRigLogicTrait, TRAIT_INTERFACE_ENUMERATOR, NULL_ANIM_TRAIT_INTERFACE_ENUMERATOR, NULL_ANIM_TRAIT_EVENT_ENUMERATOR)
+#undef TRAIT_INTERFACE_ENUMERATOR
+
+	void FRigLogicTrait::OnBecomeRelevant(FUpdateTraversalContext& Context, const TTraitBinding<IUpdate>& Binding, const FTraitUpdateState& TraitState) const
+	{
+		LLM_SCOPE_BYNAME(TEXT("Animation/RigLogic"));
+
+		const FSharedData* SharedData = Binding.GetSharedData<FSharedData>();
+		FInstanceData* InstanceData = Binding.GetInstanceData<FInstanceData>();
+
+		if (!InstanceData->Input.IsValid())
+		{
+			InstanceData->Input = Context.AllocateNodeInstance(Binding, SharedData->Input);
+		}
+
+		FRigLogicModule& Module = FModuleManager::GetModuleChecked<FRigLogicModule>("RigLogicUAF");
+		Module.DataPool.GarbageCollect();
+	}
+
+	void FRigLogicTrait::PostEvaluate(FEvaluateTraversalContext& Context, const TTraitBinding<IEvaluate>& Binding) const
+	{
+		FInstanceData* InstanceData = Binding.GetInstanceData<FInstanceData>();
+
+		if (InstanceData->Input.IsValid())
+		{
+			// TODO: Only add the task in case we want to actually want to apply facial animation for the current LOD.
+			//if (LODLevel < LODThreshold)
+			{
+				Context.AppendTask(FUAFRigLogicTask::Make(InstanceData));
+			}
+		}
+	}
+
+	void FRigLogicTrait::QueueChildrenForTraversal(FUpdateTraversalContext& Context, const TTraitBinding<IUpdateTraversal>& Binding, const FTraitUpdateState& TraitState, FUpdateTraversalQueue& TraversalQueue) const
+	{
+		const FInstanceData* InstanceData = Binding.GetInstanceData<FInstanceData>();
+
+		if (InstanceData->Input.IsValid())
+		{
+			TraversalQueue.Push(InstanceData->Input, TraitState);
+		}
+	}
+
+	uint32 FRigLogicTrait::GetNumChildren(const FExecutionContext& Context, const TTraitBinding<IHierarchy>& Binding) const
+	{
+		return 1;
+	}
+
+	void FRigLogicTrait::GetChildren(const FExecutionContext& Context, const TTraitBinding<IHierarchy>& Binding, FChildrenArray& Children) const
+	{
+		const FInstanceData* InstanceData = Binding.GetInstanceData<FInstanceData>();
+		Children.Add(InstanceData->Input);
+	}
+} // namespace UE::UAF
\ No newline at end of file
diff --git a/Engine/Plugins/Experimental/RigLogicUAF/Source/RigLogicUAF/Private/RigLogicUAF.cpp b/Engine/Plugins/Experimental/RigLogicUAF/Source/RigLogicUAF/Private/RigLogicUAF.cpp
new file mode 100644
--- /dev/null
+++ b/Engine/Plugins/Experimental/RigLogicUAF/Source/RigLogicUAF/Private/RigLogicUAF.cpp
@@ -0,0 +1,18 @@
+// Copyright Epic Games, Inc. All Rights Reserved.
+
+#include "RigLogicUAF.h"
+
+DEFINE_LOG_CATEGORY(LogRigLogicUAF);
+
+namespace UE::UAF
+{
+    void FRigLogicModule::StartupModule()
+    {
+    }
+
+    void FRigLogicModule::ShutdownModule()
+    {
+    }
+} // namespace UE::UAF
+
+IMPLEMENT_MODULE(UE::UAF::FRigLogicModule, RigLogicUAF)
\ No newline at end of file
diff --git a/Engine/Plugins/Experimental/RigLogicUAF/Source/RigLogicUAF/Public/RigLogicInstanceData.h b/Engine/Plugins/Experimental/RigLogicUAF/Source/RigLogicUAF/Public/RigLogicInstanceData.h
new file mode 100644
--- /dev/null
+++ b/Engine/Plugins/Experimental/RigLogicUAF/Source/RigLogicUAF/Public/RigLogicInstanceData.h
@@ -0,0 +1,54 @@
+// Copyright Epic Games, Inc. All Rights Reserved.
+
+#pragma once
+
+#include "DNAIndexMapping.h"
+#include "RigInstance.h"
+#include "LODPose.h"
+#include "Engine/SkeletalMesh.h"
+
+struct FSharedRigRuntimeContext;
+struct FDNAIndexMapping;
+
+namespace UE::UAF
+{
+	struct FRigLogicBoneMapping
+	{
+		uint16 RigLogicJointIndex; /** RigLogic joint index based on the internal RigLogic skeleton. Bone indices differ from the skeleton or skeletal mesh bone indices and need translation. */
+		int32 PoseBoneIndex; /** Pose bone index based on the given LOD level. */
+	};
+
+	struct FPoseBoneControlAttributeMapping
+	{
+		int32 PoseBoneIndex; /** Pose bone index based on the given LOD level. */
+		int32 DNAJointIndex;
+		int32 RotationX;
+		int32 RotationY;
+		int32 RotationZ;
+		int32 RotationW;
+	};
+
+	// Instance data unique per reference pose, cloned by the maximum number of parallel evaluations. TODO: Better call it PoolData?
+	struct FRigLogicInstanceData
+	{
+		/** Cached pointer to the shared RigLogic runtime context originally owned by UDNAAsset. */
+		TSharedPtr<FSharedRigRuntimeContext> CachedRigRuntimeContext;
+
+		/** Cached pointer to the DNA index mapping which is originally owned by UDNAAsset. */
+		TSharedPtr<FDNAIndexMapping> CachedDNAIndexMapping;
+
+		/** Bone index mapping from a RigLogic joint index to the reference skeleton bone index, one per LOD level. */
+		TArray<TArray<FRigLogicBoneMapping>> RigLogicToSkeletonBoneIndexMappingPerLOD;
+
+		TArray<TArray<FPoseBoneControlAttributeMapping>> SparseDriverJointsToControlAttributesMapPerLOD;
+		TArray<TArray<FPoseBoneControlAttributeMapping>> DenseDriverJointsToControlAttributesMapPerLOD;
+
+		uint32 NumLODs = 0;
+
+		void Init(const UE::UAF::FReferencePose* ReferencePose);
+
+	private:
+		void InitBoneIndexMapping(const UE::UAF::FReferencePose* ReferencePose);
+		void InitSparseAndDenseDriverJointMapping(const UE::UAF::FReferencePose* ReferencePose);
+	};
+} // namespace UE::UAF
\ No newline at end of file
diff --git a/Engine/Plugins/Experimental/RigLogicUAF/Source/RigLogicUAF/Public/RigLogicInstanceDataPool.h b/Engine/Plugins/Experimental/RigLogicUAF/Source/RigLogicUAF/Public/RigLogicInstanceDataPool.h
new file mode 100644
--- /dev/null
+++ b/Engine/Plugins/Experimental/RigLogicUAF/Source/RigLogicUAF/Public/RigLogicInstanceDataPool.h
@@ -0,0 +1,22 @@
+// Copyright Epic Games, Inc. All Rights Reserved.
+
+#pragma once
+
+#include "RigLogicInstanceData.h"
+
+namespace UE::UAF
+{
+	class FRigLogicInstanceDataPool
+	{
+	public:
+		TMap<TWeakObjectPtr<const USkeletalMesh>, TArray<TSharedPtr<FRigLogicInstanceData>>> Datas;
+
+		mutable FCriticalSection PoolAccessCriticalSection;
+
+		TSharedPtr<FRigLogicInstanceData> RequestData(const UE::UAF::FReferencePose* InReferencePose);
+		void FreeData(TWeakObjectPtr<const USkeletalMesh> SkeletalMesh, TSharedPtr<FRigLogicInstanceData> InData);
+
+		void GarbageCollect();
+		void Log();
+	};
+} // namespace UE::UAF
\ No newline at end of file
diff --git a/Engine/Plugins/Experimental/RigLogicUAF/Source/RigLogicUAF/Public/RigLogicTask.h b/Engine/Plugins/Experimental/RigLogicUAF/Source/RigLogicUAF/Public/RigLogicTask.h
new file mode 100644
--- /dev/null
+++ b/Engine/Plugins/Experimental/RigLogicUAF/Source/RigLogicUAF/Public/RigLogicTask.h
@@ -0,0 +1,57 @@
+// Copyright Epic Games, Inc. All Rights Reserved.
+
+#pragma once
+
+#include "CoreMinimal.h"
+#include "EvaluationVM/EvaluationTask.h"
+#include "Animation/AnimCurveTypes.h"
+#include "LODPose.h"
+#include "AnimNextStats.h"
+
+#include "RigInstance.h"
+#include "RigLogic.h"
+#include "RigLogicTrait.h"
+#include "RigLogicInstanceData.h"
+
+#include "RigLogicTask.generated.h"
+
+struct FSharedRigRuntimeContext;
+struct FDNAIndexMapping;
+
+DECLARE_CYCLE_STAT_EXTERN(TEXT("UAF Task: RigLogic"), STAT_UAF_Task_RigLogic, STATGROUP_AnimNext, RIGLOGICUAF_API);
+
+USTRUCT()
+struct RIGLOGICUAF_API FUAFRigLogicTask : public FAnimNextEvaluationTask
+{
+	GENERATED_BODY()
+
+	DECLARE_ANIM_EVALUATION_TASK(FUAFRigLogicTask)
+
+	static FUAFRigLogicTask Make(UE::UAF::FRigLogicTrait::FInstanceData* InstanceData);
+
+	// Task entry point
+	virtual void Execute(UE::UAF::FEvaluationVM& VM) const override;
+
+	using FCachedIndexedCurve = TBaseBlendedCurve<FDefaultAllocator, UE::Anim::FCurveElementIndexed>;
+	static constexpr uint16 ATTR_COUNT_PER_JOINT = 10;
+
+	void UpdateControlCurves(FRigInstance* InRigInstance,
+		const FBlendedCurve& Curves,
+		const FCachedIndexedCurve& ControlAttributeCurves,
+		const FCachedIndexedCurve& NeuralNetworkMaskCurves,
+		const TArrayView<const float> NeutralJointValues,
+		const TArrayView<const UE::UAF::FPoseBoneControlAttributeMapping> SparseDriverJointsToControlAttributesMap,
+		const TArrayView<const UE::UAF::FPoseBoneControlAttributeMapping> DenseDriverJointsToControlAttributesMap,
+		const UE::UAF::FLODPoseStack& InputPose) const;
+
+	static void UpdateJoints(const TArrayView<const UE::UAF::FRigLogicBoneMapping> RigLogicIndexToMeshIndexMapping,
+		const TArrayView<const uint16> VariableJointIndices,
+		const TArrayView<const float> NeutralJointValues,
+		const TArrayView<const float> DeltaJointValues,
+		UE::UAF::FLODPoseStack& OutputPose);
+
+	static void UpdateBlendShapeCurves(const FCachedIndexedCurve& MorphTargetCurves, const TArrayView<const float> BlendShapeValues, FBlendedCurve& OutputCurves);
+	static void UpdateAnimMapCurves(const FCachedIndexedCurve& MaskMultiplierCurves, const TArrayView<const float> AnimMapOutputs, FBlendedCurve& OutputCurves);
+
+	UE::UAF::FRigLogicTrait::FInstanceData* TraitInstanceData = nullptr;
+};
\ No newline at end of file
diff --git a/Engine/Plugins/Experimental/RigLogicUAF/Source/RigLogicUAF/Public/RigLogicTrait.h b/Engine/Plugins/Experimental/RigLogicUAF/Source/RigLogicUAF/Public/RigLogicTrait.h
new file mode 100644
--- /dev/null
+++ b/Engine/Plugins/Experimental/RigLogicUAF/Source/RigLogicUAF/Public/RigLogicTrait.h
@@ -0,0 +1,77 @@
+// Copyright Epic Games, Inc. All Rights Reserved.
+
+#pragma once
+
+#include "CoreMinimal.h"
+
+#include "RigInstance.h"
+#include "TraitCore/Trait.h"
+#include "TraitInterfaces/IContinuousBlend.h"
+#include "TraitInterfaces/IEvaluate.h"
+#include "TraitInterfaces/IHierarchy.h"
+#include "TraitInterfaces/IUpdate.h"
+
+#include "RigLogicTrait.generated.h"
+
+USTRUCT(meta = (DisplayName = "RigLogic"))
+struct RIGLOGICUAF_API FUAFRigLogicTraitSharedData : public FAnimNextTraitSharedData
+{
+	GENERATED_BODY()
+
+	UPROPERTY()
+	FAnimNextTraitHandle Input;
+
+	/*
+	 * Max LOD level that post-process AnimBPs are evaluated.
+	 * For example if you have the threshold set to 2, it will evaluate until including LOD 2 (based on 0 index). In case the LOD level gets set to 3, it will stop evaluating the post-process AnimBP.
+	 * Setting it to -1 will always evaluate it and disable LODing.
+	 */
+	UPROPERTY(EditAnywhere, Category = RigLogic)
+	int32 LODThreshold = INDEX_NONE;
+
+	// Latent pin support boilerplate
+#define TRAIT_LATENT_PROPERTIES_ENUMERATOR(GeneratorMacro) \
+		GeneratorMacro(LODThreshold)
+
+	GENERATE_TRAIT_LATENT_PROPERTIES(FUAFRigLogicTraitSharedData, TRAIT_LATENT_PROPERTIES_ENUMERATOR)
+#undef TRAIT_LATENT_PROPERTIES_ENUMERATOR
+};
+
+namespace UE::UAF
+{
+	/**
+	 * FRigLogicTrait
+	 * 
+	 * A trait that can run RigLogic.
+	 */
+	struct RIGLOGICUAF_API FRigLogicTrait : FBaseTrait, IEvaluate, IUpdate, IUpdateTraversal, IHierarchy
+	{
+		DECLARE_ANIM_TRAIT(FRigLogicTrait, FBaseTrait)
+
+		using FSharedData = FUAFRigLogicTraitSharedData;
+
+		// Do we need this? Or is the input in the shared data sufficient?
+		struct FInstanceData : FTrait::FInstanceData
+		{
+			/** Input node from which we receive the input pose as well as the facial expression curves. */
+			FTraitPtr Input;
+
+			/** Actually cloned RigLogic instance owned by this class. */
+			TUniquePtr<FRigInstance> RigInstance;
+		};
+
+		// IEvaluate impl
+		virtual void PostEvaluate(FEvaluateTraversalContext& Context, const TTraitBinding<IEvaluate>& Binding) const override;
+
+		// IUpdate impl
+		// Called before the first update when a trait stack becomes relevant
+		virtual void OnBecomeRelevant(FUpdateTraversalContext& Context, const TTraitBinding<IUpdate>& Binding, const FTraitUpdateState& TraitState) const override;
+
+		// IUpdateTraversal impl
+		virtual void QueueChildrenForTraversal(FUpdateTraversalContext& Context, const TTraitBinding<IUpdateTraversal>& Binding, const FTraitUpdateState& TraitState, FUpdateTraversalQueue& TraversalQueue) const override;
+
+		// IHierarchy impl
+		virtual uint32 GetNumChildren(const FExecutionContext& Context, const TTraitBinding<IHierarchy>& Binding) const override;
+		virtual void GetChildren(const FExecutionContext& Context, const TTraitBinding<IHierarchy>& Binding, FChildrenArray& Children) const override;
+	};
+} // namespace UE::UAF
\ No newline at end of file
diff --git a/Engine/Plugins/Experimental/RigLogicUAF/Source/RigLogicUAF/Public/RigLogicUAF.h b/Engine/Plugins/Experimental/RigLogicUAF/Source/RigLogicUAF/Public/RigLogicUAF.h
new file mode 100644
--- /dev/null
+++ b/Engine/Plugins/Experimental/RigLogicUAF/Source/RigLogicUAF/Public/RigLogicUAF.h
@@ -0,0 +1,21 @@
+// Copyright Epic Games, Inc. All Rights Reserved.
+
+#pragma once
+
+#include "Modules/ModuleManager.h"
+#include "RigLogicInstanceDataPool.h"
+
+DECLARE_LOG_CATEGORY_EXTERN(LogRigLogicUAF, Log, All);
+
+namespace UE::UAF
+{
+	class FRigLogicModule : public IModuleInterface
+	{
+	public:
+		/** IModuleInterface implementation */
+		virtual void StartupModule() override;
+		virtual void ShutdownModule() override;
+
+		UE::UAF::FRigLogicInstanceDataPool DataPool;
+	};
+} // namespace UE::UAF
\ No newline at end of file
diff --git a/Engine/Plugins/Experimental/RigLogicUAF/Source/RigLogicUAF/RigLogicUAF.Build.cs b/Engine/Plugins/Experimental/RigLogicUAF/Source/RigLogicUAF/RigLogicUAF.Build.cs
new file mode 100644
--- /dev/null
+++ b/Engine/Plugins/Experimental/RigLogicUAF/Source/RigLogicUAF/RigLogicUAF.Build.cs
@@ -0,0 +1,54 @@
+// Copyright Epic Games, Inc. All Rights Reserved.
+
+using UnrealBuildTool;
+
+public class RigLogicUAF : ModuleRules
+{
+	public RigLogicUAF(ReadOnlyTargetRules Target) : base(Target)
+	{
+		PCHUsage = ModuleRules.PCHUsageMode.UseExplicitOrSharedPCHs;
+		
+		PublicIncludePaths.AddRange(
+			new string[]
+			{
+			}
+			);
+				
+		
+		PrivateIncludePaths.AddRange(
+			new string[]
+			{
+			}
+			);
+			
+		
+		PublicDependencyModuleNames.AddRange(
+			new string[]
+			{
+				"Core",
+				"Engine",
+				"RigLogicModule",
+				"UAF",
+				"UAFAnimGraph"
+			}
+			);
+			
+		
+		PrivateDependencyModuleNames.AddRange(
+			new string[]
+			{
+				"CoreUObject",
+				"Engine",
+				"Slate",
+				"SlateCore",
+			}
+			);
+		
+		
+		DynamicallyLoadedModuleNames.AddRange(
+			new string[]
+			{
+			}
+			);
+	}
+}
diff --git a/Engine/Plugins/Experimental/RigLogicUAF/Source/RigLogicUAFUncookedOnly/Private/RigLogicUAFUncookedOnly.cpp b/Engine/Plugins/Experimental/RigLogicUAF/Source/RigLogicUAFUncookedOnly/Private/RigLogicUAFUncookedOnly.cpp
new file mode 100644
--- /dev/null
+++ b/Engine/Plugins/Experimental/RigLogicUAF/Source/RigLogicUAFUncookedOnly/Private/RigLogicUAFUncookedOnly.cpp
@@ -0,0 +1,45 @@
+// Copyright Epic Games, Inc. All Rights Reserved.
+
+#include "RigLogicUAFUncookedOnly.h"
+#include "Engine/Texture2D.h"
+#include "Styling/SlateStyle.h"
+#include "Styling/SlateStyleRegistry.h"
+#include "Brushes/SlateImageBrush.h"
+#include "Interfaces/IPluginManager.h"
+
+DEFINE_LOG_CATEGORY(LogRigLogicUAFUncookedOnly);
+
+namespace UE::UAF
+{
+    void FRigLogicModuleUncookedOnly::StartupModule()
+    {
+    }
+
+    void FRigLogicModuleUncookedOnly::ShutdownModule()
+    {
+    }
+
+	const FSlateBrush& FRigLogicModuleUncookedOnly::GetIcon()
+	{
+    	static TSharedPtr<FSlateStyleSet> StyleSet = nullptr;
+    	static const FSlateBrush* RegisteredBrush = nullptr;
+
+    	if (!StyleSet.IsValid())
+    	{
+    		StyleSet = MakeShareable(new FSlateStyleSet("RigLogicStyleUAF"));
+
+    		const FString PluginPath = IPluginManager::Get().FindPlugin("RigLogicUAF")->GetBaseDir();
+    		StyleSet->SetContentRoot(PluginPath / TEXT("Resources"));
+
+    		const FName BrushName = "RigLogicUAF_Icon";
+    		StyleSet->Set(BrushName, new FSlateImageBrush(StyleSet->RootToContentDir(TEXT("Icon128"), TEXT(".png")), FVector2D(16, 16)));
+
+    		FSlateStyleRegistry::RegisterSlateStyle(*StyleSet);
+    		RegisteredBrush = StyleSet->GetBrush(BrushName);
+    	}
+
+    	return *RegisteredBrush;
+	}
+} // namespace UE::UAF
+
+IMPLEMENT_MODULE(UE::UAF::FRigLogicModuleUncookedOnly, RigLogicUAFUncookedOnly)
\ No newline at end of file
diff --git a/Engine/Plugins/Experimental/RigLogicUAF/Source/RigLogicUAFUncookedOnly/Public/RigLogicUAFUncookedOnly.h b/Engine/Plugins/Experimental/RigLogicUAF/Source/RigLogicUAFUncookedOnly/Public/RigLogicUAFUncookedOnly.h
new file mode 100644
--- /dev/null
+++ b/Engine/Plugins/Experimental/RigLogicUAF/Source/RigLogicUAFUncookedOnly/Public/RigLogicUAFUncookedOnly.h
@@ -0,0 +1,21 @@
+// Copyright Epic Games, Inc. All Rights Reserved.
+
+#pragma once
+
+#include "Modules/ModuleManager.h"
+#include "Textures/SlateIcon.h"
+
+DECLARE_LOG_CATEGORY_EXTERN(LogRigLogicUAFUncookedOnly, Log, All);
+
+namespace UE::UAF
+{
+	class FRigLogicModuleUncookedOnly : public IModuleInterface
+	{
+	public:
+		/** IModuleInterface implementation */
+		virtual void StartupModule() override;
+		virtual void ShutdownModule() override;
+
+		static const FSlateBrush& GetIcon();
+	};
+} // namespace UE::UAF
\ No newline at end of file
diff --git a/Engine/Plugins/Experimental/RigLogicUAF/Source/RigLogicUAFUncookedOnly/Public/UAFGraphNodeTemplate_RigLogic.h b/Engine/Plugins/Experimental/RigLogicUAF/Source/RigLogicUAFUncookedOnly/Public/UAFGraphNodeTemplate_RigLogic.h
new file mode 100644
--- /dev/null
+++ b/Engine/Plugins/Experimental/RigLogicUAF/Source/RigLogicUAFUncookedOnly/Public/UAFGraphNodeTemplate_RigLogic.h
@@ -0,0 +1,53 @@
+﻿// Copyright Epic Games, Inc. All Rights Reserved.
+
+#pragma once
+
+#include "AnimGraphUncookedOnlyUtils.h"
+#include "Templates/UAFGraphNodeTemplate.h"
+#include "RigLogicTrait.h"
+#include "Traits/PassthroughBlendTrait.h"
+#include "RigLogicUAFUncookedOnly.h"
+
+#include "UAFGraphNodeTemplate_RigLogic.generated.h"
+
+#define LOCTEXT_NAMESPACE "UAFGraphNodeTemplate_RigLogic"
+
+UCLASS()
+class UUAFGraphNodeTemplate_RigLogic : public UUAFGraphNodeTemplate
+{
+	GENERATED_BODY()
+
+	UUAFGraphNodeTemplate_RigLogic()
+	{
+		Title = LOCTEXT("RigLogicUAFTitle", "RigLogic");
+		TooltipText = LOCTEXT("RigLogicUAFTooltip", "RigLogic\n"
+			"Performs facial animation and drives body correctives where applicable.\n"
+			"Input: Input Pose\n"
+			"Output: Overwritten joint transforms, and curves to drive blend shapes and animated maps.");
+		Category = LOCTEXT("RigLogicUAFCategory", "UAF");
+		MenuDescription = LOCTEXT("RigLogicUAFMenuDesc", "RigLogic");
+		Color = FLinearColor(FColor(38, 187, 255));
+		Icon = UE::UAF::FRigLogicModuleUncookedOnly::GetIcon();
+
+		FPassthroughBlendTraitSharedData PassthroughData;
+		PassthroughData.AlphaInputType = EAnimAlphaInputType::Bool;
+
+		Traits =
+		{
+			TInstancedStruct<FUAFRigLogicTraitSharedData>::Make(),
+			TInstancedStruct<FPassthroughBlendTraitSharedData>::Make(PassthroughData)
+		};
+		SetCategoryForPinsInLayout(
+			{
+				GET_PIN_PATH_STRING_CHECKED(FUAFRigLogicTraitSharedData, Input),
+				GET_PIN_PATH_STRING_CHECKED(FPassthroughBlendTraitSharedData, bAlphaBoolEnabled)
+			},
+			TEXT("Default"),
+			NodeLayout,
+			true);
+		SetDisplayNameForPinInLayout(GET_PIN_PATH_STRING_CHECKED(FUAFRigLogicTraitSharedData, Input), TEXT("Input"), NodeLayout);
+		SetDisplayNameForPinInLayout(GET_PIN_PATH_STRING_CHECKED(FPassthroughBlendTraitSharedData, bAlphaBoolEnabled), TEXT("Enabled"), NodeLayout);
+	}
+};
+
+#undef LOCTEXT_NAMESPACE
\ No newline at end of file
diff --git a/Engine/Plugins/Experimental/RigLogicUAF/Source/RigLogicUAFUncookedOnly/RigLogicUAFUncookedOnly.Build.cs b/Engine/Plugins/Experimental/RigLogicUAF/Source/RigLogicUAFUncookedOnly/RigLogicUAFUncookedOnly.Build.cs
new file mode 100644
--- /dev/null
+++ b/Engine/Plugins/Experimental/RigLogicUAF/Source/RigLogicUAFUncookedOnly/RigLogicUAFUncookedOnly.Build.cs
@@ -0,0 +1,57 @@
+// Copyright Epic Games, Inc. All Rights Reserved.
+
+using UnrealBuildTool;
+
+public class RigLogicUAFUncookedOnly : ModuleRules
+{
+	public RigLogicUAFUncookedOnly(ReadOnlyTargetRules Target) : base(Target)
+	{
+		PCHUsage = ModuleRules.PCHUsageMode.UseExplicitOrSharedPCHs;
+		
+		PublicIncludePaths.AddRange(
+			new string[]
+			{
+			}
+			);
+				
+		
+		PrivateIncludePaths.AddRange(
+			new string[]
+			{
+			}
+			);
+			
+		
+		PublicDependencyModuleNames.AddRange(
+			new string[]
+			{
+				"Core",
+				"Engine",
+				"RigLogicModule",
+				"UAF",
+				"UAFAnimGraph",
+				"RigLogicUAF",
+				"UAFAnimGraphUncookedOnly",
+				"Projects"
+			}
+			);
+			
+		
+		PrivateDependencyModuleNames.AddRange(
+			new string[]
+			{
+				"CoreUObject",
+				"Engine",
+				"Slate",
+				"SlateCore",
+			}
+			);
+		
+		
+		DynamicallyLoadedModuleNames.AddRange(
+			new string[]
+			{
+			}
+			);
+	}
+}
