RigLogicUAF: Support new UDNAAssetUserData in addition to legacy UDNAAsset

FRigLogicInstanceData::Init() only looked for UDNAAsset (legacy) as asset user data on the skeletal mesh. Newer MetaHuman setups use UDNAAssetUserData which wraps a UDNA object reference. Now tries UDNAAssetUserData first, falls back to UDNAAsset.

#Synced-CL 51755844

diff --git a/Engine/Plugins/Experimental/RigLogicUAF/Source/RigLogicUAF/Private/RigLogicInstanceData.cpp b/Engine/Plugins/Experimental/RigLogicUAF/Source/RigLogicUAF/Private/RigLogicInstanceData.cpp
--- a/Engine/Plugins/Experimental/RigLogicUAF/Source/RigLogicUAF/Private/RigLogicInstanceData.cpp
+++ b/Engine/Plugins/Experimental/RigLogicUAF/Source/RigLogicUAF/Private/RigLogicInstanceData.cpp
@@ -5,7 +5,9 @@
 
 #include "Animation/Skeleton.h"
 #include "Engine/SkeletalMesh.h"
+#include "DNA.h"
 #include "DNAAsset.h"
+#include "DNAAssetUserData.h"
 #include "LODPose.h"
 #include "SharedRigRuntimeContext.h"
 
@@ -35,17 +37,25 @@ namespace UE::UAF
 			return;
 		}
 
-		UDNAAsset* DNAAsset = Cast<UDNAAsset>(SkeletalMesh->GetAssetUserDataOfClass(UDNAAsset::StaticClass()));
-		if (!DNAAsset)
+		// Resolve DNA: try the new UDNAAssetUserData first, then fall back to the legacy UDNAAsset
+		TSharedPtr<FSharedRigRuntimeContext> SharedRigRuntimeContext;
+		if (UDNAAssetUserData* DNAUserData = Cast<UDNAAssetUserData>(SkeletalMesh->GetAssetUserDataOfClass(UDNAAssetUserData::StaticClass())))
 		{
-			UE_LOG(LogRigLogicUAF, Warning, TEXT("No DNA asset assigned to the skeletal mesh."));
-			return;
+			if (UDNA* DNA = DNAUserData->DNAAsset)
+			{
+				SharedRigRuntimeContext = DNA->GetRigRuntimeContext();
+				CachedDNAIndexMapping = DNA->GetDNAIndexMapping(Skeleton, SkeletalMesh);
+			}
+		}
+		else if (UDNAAsset* DNAAsset = Cast<UDNAAsset>(SkeletalMesh->GetAssetUserDataOfClass(UDNAAsset::StaticClass())))
+		{
+			SharedRigRuntimeContext = DNAAsset->GetRigRuntimeContext();
+			CachedDNAIndexMapping = DNAAsset->GetDNAIndexMapping(Skeleton, SkeletalMesh);
 		}
 
-		TSharedPtr<FSharedRigRuntimeContext> SharedRigRuntimeContext = DNAAsset->GetRigRuntimeContext();
 		if (!SharedRigRuntimeContext.IsValid())
 		{
-			UE_LOG(LogRigLogicUAF, Error, TEXT("Can't get the shared rig runtime context."));
+			UE_LOG(LogRigLogicUAF, Warning, TEXT("No DNA asset assigned to the skeletal mesh."));
 			return;
 		}
 
@@ -54,8 +64,6 @@ namespace UE::UAF
 			CachedRigRuntimeContext = SharedRigRuntimeContext;
 		}
 
-		CachedDNAIndexMapping = DNAAsset->GetDNAIndexMapping(Skeleton, SkeletalMesh);
-
 		NumLODs = CachedRigRuntimeContext->VariableJointIndicesPerLOD.Num();
 
 		InitBoneIndexMapping(ReferencePose);
